Sequencer Musical Mode grid markers

#tests Editor
#rnx

diff --git a/Engine/Source/Editor/Sequencer/Private/SequencerTimeSliderController.cpp b/Engine/Source/Editor/Sequencer/Private/SequencerTimeSliderController.cpp
--- a/Engine/Source/Editor/Sequencer/Private/SequencerTimeSliderController.cpp
+++ b/Engine/Source/Editor/Sequencer/Private/SequencerTimeSliderController.cpp
@@ -499,8 +499,14 @@ int32 FSequencerTimeSliderController::OnPaintTimeSlider( bool bMirrorLabels, con
 		FScrubRangeToScreen RangeToScreen( LocalViewRange, AllottedGeometry.Size );
 
 		// draw tick marks
-		const float MajorTickHeight = 9.0f;
-	
+		float MajorTickHeight = 9.0f;
+
+		TViewModelPtr<IClockExtension> Clock = Sequencer->GetViewModel()->GetRootSequenceModel().ImplicitCast();
+		if (Clock)
+		{
+			MajorTickHeight = Clock->GetCustomMajorTickSize(Sequencer);
+		}
+
 		FDrawTickArgs Args;
 		{
 			Args.AllottedGeometry = AllottedGeometry;
diff --git a/Engine/Source/Editor/Sequencer/Public/MVVM/Extensions/IClockExtension.h b/Engine/Source/Editor/Sequencer/Public/MVVM/Extensions/IClockExtension.h
--- a/Engine/Source/Editor/Sequencer/Public/MVVM/Extensions/IClockExtension.h
+++ b/Engine/Source/Editor/Sequencer/Public/MVVM/Extensions/IClockExtension.h
@@ -33,6 +33,10 @@ public:
 	{
 		return 22.0;
 	}
+	virtual float GetCustomMajorTickSize(TSharedPtr<ISequencer> Sequencer) const
+	{
+		return 9.0;
+	}
 	virtual bool DrawTicks(TSharedPtr<ISequencer> Sequencer, FSlateWindowElementList& OutDrawElements, const TRange<double>& ViewRange, const FSequencerTimeSliderController::FScrubRangeToScreen& RangeToScreen, FSequencerTimeSliderController::FDrawTickArgs& InOutArgs) const
 	{
 		return false;
diff --git a/Engine/Source/Runtime/MusicEnvironment/Private/FrameBasedMusicMap.cpp b/Engine/Source/Runtime/MusicEnvironment/Private/FrameBasedMusicMap.cpp
--- a/Engine/Source/Runtime/MusicEnvironment/Private/FrameBasedMusicMap.cpp
+++ b/Engine/Source/Runtime/MusicEnvironment/Private/FrameBasedMusicMap.cpp
@@ -361,8 +361,12 @@ FFrameTime UFrameBasedMusicMap::RemapFrameTime(FFrameTime SourceFrameTime, const
 	{
 		return SourceFrameTime;
 	}
+	
+	float Seconds = FrameResolution.AsSeconds(SourceFrameTime);
+	const FFrameBasedTempoPoint* TempoPoint = GetTempoPointForSeconds(Seconds);
+	int32 Tick = TempoPoint->SecondsToTick(Seconds);
 
-	FMusicalTime MusicalTime = OldMap->GetMusicalTime(SourceFrameTime);
+	FMusicalTime MusicalTime = OldMap->GetMusicalTime(Tick);
 	const FFrameBasedTimeSignaturePoint* TimeSignaturePoint = GetTimeSignaturePointForBar(MusicalTime.Bar);
 	if (TimeSignaturePoint->TicksPerBar() != MusicalTime.TicksPerBar)
 	{
@@ -470,61 +474,63 @@ void UFrameBasedMusicMap::AddTempo(int32 AtTick, float AtMs, float Bpm)
 	}
 }
 
-TArray<FSubdivisionPoint> UFrameBasedMusicMap::GatherSubdivisonsForTimespan(const FFrameTime& Start, const FFrameTime& End, const FFrameTime& MinSpacing, int32 QuantizationInterval) const
+void UFrameBasedMusicMap::GatherSubdivisonsForTimespan(int32 StartTick, int32 EndTick, int32 QuantizationInterval, TArray<FGridInterval>& OutSubdivisions ) const
 {
-	TArray<FSubdivisionPoint> Result;
-
-	FMusicalTime StartMusicalTime = Quantize(GetMusicalTime(Start), QuantizationInterval, EQuantizeDirection::Ceiling);
-	FMusicalTime EndMusicalTime = Quantize(GetMusicalTime(End), QuantizationInterval, EQuantizeDirection::Floor);
+	FMusicalTime StartMusicalTime = Quantize(GetMusicalTime(StartTick), QuantizationInterval, EQuantizeDirection::Ceiling);
+	FMusicalTime EndMusicalTime = Quantize(GetMusicalTime(EndTick), QuantizationInterval, EQuantizeDirection::Floor);
 
 	if (StartMusicalTime >= EndMusicalTime)
 	{
-		return Result;
+		return;
 	}
 
+	while (StartMusicalTime <= EndMusicalTime)
+	{
+		float SubdivisionTime = GetSeconds(StartMusicalTime);
+
+		const FFrameBasedTimeSignaturePoint* TimeSignatureAtFrame = GetTimeSignaturePointForBar(StartMusicalTime.Bar);
+		int32 Tick = TimeSignatureAtFrame->MusicalTimeToTick(StartMusicalTime);
+
+		FGridInterval GridInterval(SubdivisionTime, Tick, StartMusicalTime);
+
+		OutSubdivisions.Add(GridInterval);
+		StartMusicalTime = Add(StartMusicalTime, QuantizationInterval);
+	}
+}
 
-	Result.Add(MakeSubdivisionPoint(StartMusicalTime));
-	FFrameTime LastEntryFrameTime = GetFrameTime(StartMusicalTime);
-	int32 LastBarNumber = StartMusicalTime.Bar;
-	StartMusicalTime = Add(StartMusicalTime, QuantizationInterval);
-	if (StartMusicalTime.Bar != LastBarNumber)
+void UFrameBasedMusicMap::GatherBarsForTimespan(int32 StartBar, int32 EndBar, int32 NumBarsInterval, TArray<FGridInterval>& OutSubdivisions) const
+{
+	if (StartBar >= EndBar)
 	{
-		StartMusicalTime.TickInBar = 0;
-		LastBarNumber = StartMusicalTime.Bar;
+		return;
 	}
-	while (StartMusicalTime <= EndMusicalTime)
+
+	while (StartBar <= EndBar)
 	{
-		FFrameTime ThisEntryFrameTime = GetFrameTime(StartMusicalTime);
-		if ((ThisEntryFrameTime - LastEntryFrameTime) < MinSpacing)
-		{
-			if (StartMusicalTime.TickInBar == 0)
-			{
-				// This is a downbeat. We are too close to the last, BUT we 
-				// take precedence. So blow away the last...
-				Result.Pop(EAllowShrinking::No);
-			}
-			else
-			{
-				StartMusicalTime = Add(StartMusicalTime, QuantizationInterval);
-				if (StartMusicalTime.Bar != LastBarNumber)
-				{
-					StartMusicalTime.TickInBar = 0;
-					LastBarNumber = StartMusicalTime.Bar;
-				}
-				continue;
-			}
-		}
-		Result.Add(MakeSubdivisionPoint(StartMusicalTime));
-		LastEntryFrameTime = GetFrameTime(StartMusicalTime);
-		StartMusicalTime = Add(StartMusicalTime, QuantizationInterval);
-		if (StartMusicalTime.Bar != LastBarNumber)
-		{
-			StartMusicalTime.TickInBar = 0;
-			LastBarNumber = StartMusicalTime.Bar;
-		}
+		float BarSeconds = GetSeconds(StartBar);
+
+		const FFrameBasedTimeSignaturePoint* TimeSignatureAtFrame = GetTimeSignaturePointForBar(StartBar);
+		int32 Tick = TimeSignatureAtFrame->BarToTick(StartBar);
+		FMusicalTime MusicalTime = TimeSignatureAtFrame->TickToMusicalTime(Tick);
+
+		FGridInterval GridInterval(BarSeconds, Tick, MusicalTime);
+		OutSubdivisions.Add(GridInterval);
+
+		StartBar = StartBar + NumBarsInterval;
 	}
-	
-	return Result;
+}
+
+FMusicalTime UFrameBasedMusicMap::GetMusicalTime(int32 Tick) const
+{
+	if (IsEmpty())
+	{
+		return FMusicalTime();
+	}
+
+	const FFrameBasedTimeSignaturePoint* TimeSignatureAtFrame = GetTimeSignaturePointForTick(Tick);
+	const FFrameBasedTempoPoint* TempoAtFrame = GetTempoPointForTick(Tick);
+
+	return TimeSignatureAtFrame->TickToMusicalTime(Tick);
 }
 
 FMusicalTime UFrameBasedMusicMap::GetMusicalTime(const FFrameTime& FrameTime) const
@@ -542,15 +548,6 @@ FMusicalTime UFrameBasedMusicMap::GetMusicalTime(const FFrameTime& FrameTime) co
 	return TimeSignatureAtFrame->TickToMusicalTime(Tick);
 }
 
-FSubdivisionPoint UFrameBasedMusicMap::MakeSubdivisionPoint(const FMusicalTime& MusicalTime) const
-{
-	FSubdivisionPoint Result;
-	Result.MusicalTime = MusicalTime;
-	Result.FrameTime = GetFrameTime(MusicalTime);
-	Result.PrettyString = MakeBarBeatStringForMusicalTime(MusicalTime, FrameBasedMusicMap::EBarBeatFormat::Position);
-	return Result;
-}
-
 FString UFrameBasedMusicMap::MakeBarBeatString(double DecimalFrame, const FFrameBasedTimeSignaturePoint* TimeSignature, const FFrameBasedTempoPoint* Tempo, FrameBasedMusicMap::EBarBeatFormat BarBeatFormat) const
 {
 	double Seconds = FrameResolution.AsSeconds(FFrameTime::FromDecimal(DecimalFrame));
@@ -566,4 +563,297 @@ FString UFrameBasedMusicMap::MakeBarBeatString(double DecimalFrame, const FFrame
 	}
 	return FString::Printf(TEXT("%d.%.2f"), MusicalTime.Bar, BeatAsDouble / 100.0);
 }
+
+void UFrameBasedMusicMap::GetGridMetricsMusic(float InViewStartSeconds, float InViewEndSeconds, float PixelsPerSecond, float MinTickPx, float DesiredMajorTickPx, TArray<FGridInterval>& OutMajorIntervals, TArray<FGridInterval>& OutMinorIntervals) const
+{
+	const FFrameBasedTempoPoint* StartTempoPoint = GetTempoPointForSeconds(InViewStartSeconds);
+	int32 StartTick = StartTempoPoint->SecondsToTick(InViewStartSeconds);
+	int32 QuantizedStartTick = StartTick;
+
+	FMusicalTime StartMusicalTime = GetMusicalTime(StartTick);
+	int32 StartBar = StartMusicalTime.Bar;
+
+	int32 SubdivisionStepTickSize = FrameBasedMusicMap::EWholeQuantizationTicks::ThirtysecondNote;
+	int32 BarStepSize = 1;
+
+	const float DesiredMajorTickSeconds = DesiredMajorTickPx / PixelsPerSecond;
+
+	//We find a subdivision size that intersects with our ViewStart and fits our desired size
+	bool bFoundFirstGridLine = false;
+	while (!bFoundFirstGridLine)
+	{
+		if (SubdivisionStepTickSize <= FrameBasedMusicMap::EWholeQuantizationTicks::HalfNote)
+		{
+			QuantizedStartTick = FMath::FloorToDouble((float)StartTick / (float)SubdivisionStepTickSize) * SubdivisionStepTickSize;
+			
+			StartTempoPoint = GetTempoPointForTick(QuantizedStartTick);
+			double QuantizedStartSeconds = StartTempoPoint->TickToSeconds(QuantizedStartTick);
+
+			int32 SteppedTick = QuantizedStartTick + SubdivisionStepTickSize;
+			const FFrameBasedTempoPoint* StartSteppedTempoPoint = GetTempoPointForTick(SteppedTick);
+			double SteppedSeconds = StartSteppedTempoPoint->TickToSeconds(SteppedTick);
+
+			float StepSizeSeconds = SteppedSeconds - QuantizedStartSeconds;
+
+			//Make sure were not stepping through a downbeat of next bar
+			StartMusicalTime = GetMusicalTime(QuantizedStartTick);
+			FMusicalTime SteppedMusicalTime = GetMusicalTime(SteppedTick);
+
+			bool bValidStep = false;
+			if (StartMusicalTime.Bar == SteppedMusicalTime.Bar || SteppedMusicalTime.TickInBar == 0)
+			{
+				bValidStep = true;
+			}
+
+			if (StepSizeSeconds > DesiredMajorTickSeconds && bValidStep)
+			{
+				bFoundFirstGridLine = true;
+
+				FGridInterval NonLinearGridInterval = FGridInterval(QuantizedStartSeconds, QuantizedStartTick, GetMusicalTime(QuantizedStartTick));
+				OutMajorIntervals.Add(NonLinearGridInterval);
+			}
+			else
+			{
+				SubdivisionStepTickSize *= 2;
+			}
+		}
+		else
+		{
+			int32 QuantizedStartBar = FMath::FloorToDouble((float)StartBar / (float)BarStepSize) * BarStepSize;
+			float QuantizedBarStartTime = GetSeconds(QuantizedStartBar);
+			const FFrameBasedTempoPoint* StartSteppedTempoPoint = GetTempoPointForSeconds(QuantizedBarStartTime);
+			QuantizedStartTick = StartSteppedTempoPoint->SecondsToTick(QuantizedBarStartTime);
+
+			int32 SteppedBar = QuantizedStartBar + BarStepSize;
+
+			float SecondsAtSteppedBar = GetSeconds(SteppedBar);
+
+			float StepSizeSeconds = SecondsAtSteppedBar - QuantizedBarStartTime;
+
+			if (StepSizeSeconds > DesiredMajorTickSeconds || BarStepSize >= 2048)
+			{
+				bFoundFirstGridLine = true;
+
+				FGridInterval NonLinearGridInterval = FGridInterval(QuantizedBarStartTime, QuantizedStartTick, BarsToMusicalTime(QuantizedStartBar));
+				OutMajorIntervals.Add(NonLinearGridInterval);
+			}
+			else
+			{
+				BarStepSize *= 2;
+			}
+		}
+	}
+
+	const FFrameBasedTempoPoint* EndTempoPoint = GetTempoPointForSeconds(InViewEndSeconds);
+	int32 EndTick = EndTempoPoint->SecondsToTick(InViewEndSeconds);
+
+	int32 StepSizeTicks = 0;
+
+	for (int32 CurrentTick = QuantizedStartTick; CurrentTick < EndTick; CurrentTick += StepSizeTicks)
+	{
+		ComputeGridSpacingAtTick(
+			CurrentTick,
+			DesiredMajorTickSeconds,
+			StepSizeTicks);
+
+		int32 SteppedTicks = CurrentTick + StepSizeTicks;
+
+		FMusicalTime CurrentMusicalTime = GetMusicalTime(CurrentTick);
+		FMusicalTime SteppedMusicalTime = GetMusicalTime(SteppedTicks);
+
+		//After computing spacing, we may need to correct past intervals if the spacing dosen't make sense
+		if (SteppedMusicalTime.Bar != CurrentMusicalTime.Bar && CurrentMusicalTime.Bar > 0)
+		{
+			int32 NumBarsStepped = SteppedMusicalTime.Bar - CurrentMusicalTime.Bar;
+
+			//Make sure our step size dosen't take us into next bar, i.e. always favour marker on downbeat
+			if (SteppedMusicalTime.TickInBar > 0)
+			{
+				OutMajorIntervals.Pop();
+				
+				CurrentTick = OutMajorIntervals.Last().IntervalTick;
+				CurrentMusicalTime = OutMajorIntervals.Last().MusicalTime;
+
+				StepSizeTicks = CurrentMusicalTime.TicksPerBar - CurrentMusicalTime.TickInBar;
+				SteppedTicks = CurrentTick + StepSizeTicks;
+				SteppedMusicalTime = GetMusicalTime(SteppedTicks);
+			}
+			//Make sure we aren't stepping up in multiple number of bars at wrong interval
+			else if (SteppedMusicalTime.Bar % NumBarsStepped != 0)
+			{
+				//We need to increase the previous interval 
+				FMusicalTime PreviousSteppedMusicalTime = OutMajorIntervals.Last().MusicalTime;
+
+				OutMajorIntervals.Pop();
+
+				if (OutMajorIntervals.Num() > 0)
+				{
+					FMusicalTime PreviousMusicalTime = OutMajorIntervals.Last().MusicalTime;
+
+					NumBarsStepped = PreviousSteppedMusicalTime.Bar - PreviousMusicalTime.Bar;
+					NumBarsStepped *= 2;
+
+					int32 SteppedBars = PreviousMusicalTime.Bar + NumBarsStepped;
+
+					const FFrameBasedTimeSignaturePoint* TimeSigPoint = GetTimeSignaturePointForBar(SteppedBars);
+					int32 BarsAtTimeSig = SteppedBars - TimeSigPoint->OnBar;
+
+					StepSizeTicks = (TimeSigPoint->OnTick - CurrentTick) + BarsAtTimeSig * TimeSigPoint->TicksPerBar();
+					SteppedTicks = CurrentTick + StepSizeTicks;
+					SteppedMusicalTime = GetMusicalTime(SteppedTicks);
+				}
+			}
+		}
+
+		const FFrameBasedTempoPoint* SteppedTempoPoint = GetTempoPointForTick(SteppedTicks);
+		float SteppedSizeSeconds = SteppedTempoPoint->TickToMs(SteppedTicks) / 1000.f;
+
+		FGridInterval NonLinearGridInterval = FGridInterval(SteppedSizeSeconds, SteppedTicks, SteppedMusicalTime);
+		OutMajorIntervals.Add(NonLinearGridInterval);
+	}
+
+	for (int32 Idx = 0; Idx < OutMajorIntervals.Num() - 1; ++Idx)
+	{
+		ComputeMinorGridSpacing(OutMajorIntervals[Idx].IntervalTick, OutMajorIntervals[Idx + 1].IntervalTick, OutMinorIntervals);
+	}
+}
+
+void UFrameBasedMusicMap::ComputeGridSpacingAtTick(int32 TickToCompute, float DesiredMajorTickSeconds, int32& OutMajorIntervalTicks) const
+{
+	const FFrameBasedTempoPoint* TempoPoint = GetTempoPointForTick(TickToCompute);
+	float TimeToComputeSeconds = TempoPoint->TickToSeconds(TickToCompute);
+
+	float DesiredTimeStepped = TimeToComputeSeconds + DesiredMajorTickSeconds;
+
+	int32 TickStepped = TempoPoint->MsToTick(DesiredTimeStepped * 1000.f);
+
+	int32 StepTickSize = TickStepped - TickToCompute;
+
+	//This will round up to lowest 2^N tick size that is larger than our desired size
+	int32 PowTwo = FMath::RoundUpToPowerOfTwo(StepTickSize / FrameBasedMusicMap::EWholeQuantizationTicks::ThirtysecondNote);
+	StepTickSize = FrameBasedMusicMap::EWholeQuantizationTicks::ThirtysecondNote * PowTwo;
+
+	FMusicalTime StartMusicalTime = GetMusicalTime(TickToCompute);
+	if (StartMusicalTime.TickInBar == 0 && StepTickSize > FrameBasedMusicMap::EWholeQuantizationTicks::HalfNote)
+	{
+		FMusicalTime SteppedMusicalTime = GetMusicalTime(TickStepped);
+
+		int32 NumBarsStepped = SteppedMusicalTime.Bar - StartMusicalTime.Bar;
+		NumBarsStepped++;
+
+		NumBarsStepped = FMath::Clamp(NumBarsStepped, 1, 2048);
+		int32 NumMultipleBarsStepped = FMath::RoundUpToPowerOfTwo(NumBarsStepped);
+
+		int32 SteppedBars = StartMusicalTime.Bar + NumMultipleBarsStepped;
+
+		const FFrameBasedTimeSignaturePoint* TimeSigPoint = GetTimeSignaturePointForBar(SteppedBars);
+		int32 BarsAtTimeSig = SteppedBars - TimeSigPoint->OnBar;
+
+		StepTickSize = (TimeSigPoint->OnTick - TickToCompute) + BarsAtTimeSig * TimeSigPoint->TicksPerBar();
+	}
+	else
+	{
+		//We can't step upwards if our starting tick isn't on a subdivision boundary of that size
+		//I.e. if we are on the tick 3/8th notes into a bar, we can't start stepping in 1/4 notes
+		while (StartMusicalTime.TickInBar % StepTickSize != 0 && StepTickSize > FrameBasedMusicMap::EWholeQuantizationTicks::ThirtysecondNote)
+		{
+			StepTickSize /= 2;
+		}
+
+		if (StepTickSize < FrameBasedMusicMap::EWholeQuantizationTicks::ThirtysecondNote)
+		{
+			StepTickSize = FrameBasedMusicMap::EWholeQuantizationTicks::ThirtysecondNote;
+		}
+	}
+
+	OutMajorIntervalTicks = StepTickSize;
+}
+
+void UFrameBasedMusicMap::ComputeMinorGridSpacing(int32 StartTick, int32 EndTick, TArray<FGridInterval>& OutMinorIntervals) const
+{
+	int32 StepTickSize = EndTick - StartTick;
+
+	FMusicalTime StartMusicalTime = GetMusicalTime(StartTick);
+	FMusicalTime EndMusicalTime = GetMusicalTime(EndTick);
+
+	int32 BarStepSize = EndMusicalTime.Bar - StartMusicalTime.Bar;
+
+	if (StepTickSize >= StartMusicalTime.TicksPerBar && BarStepSize >= 4)
+	{
+		int32 MinorIntervalBars = BarStepSize / 4;
+		GatherBarsForTimespan(StartMusicalTime.Bar, EndMusicalTime.Bar, MinorIntervalBars, OutMinorIntervals);
+	}
+	else
+	{
+		int32 MinorTickSize = StepTickSize / 4;
+		MinorTickSize = FMath::Clamp(MinorTickSize, FrameBasedMusicMap::EWholeQuantizationTicks::ThirtysecondNote, FrameBasedMusicMap::EWholeQuantizationTicks::HalfNote);
+		GatherSubdivisonsForTimespan(StartTick, EndTick, MinorTickSize, OutMinorIntervals);
+	}
+}
+
+TArray<FFrameBasedTempoPoint> UFrameBasedMusicMap::GetTempoMetrics(float InViewStartSeconds, float InViewEndSeconds) const
+{
+	TArray<FFrameBasedTempoPoint> TempoPoints;
+
+	float ViewStartMs = InViewStartSeconds * 1000.f;
+
+	float ViewEndMs = InViewEndSeconds * 1000.f;
+
+	for (FFrameBasedTempoPoint TempoPoint : TempoMap)
+	{
+		if (TempoPoint.OnMs > ViewStartMs && TempoPoint.OnMs < ViewEndMs)
+		{
+			TempoPoints.Add(TempoPoint);
+		}
+	}
+
+	return TempoPoints;
+}
+
+void UFrameBasedMusicMap::GetTempoTimeSignatureMarkers(float InStartSeconds, float InEndSeconds, TArray<FTempoTimeSignatureMarker>& OutMarkers) const
+{
+	float ViewStartMs = InStartSeconds * 1000.f;
+	float ViewEndMs = InEndSeconds * 1000.f;
+
+	const FFrameBasedTempoPoint* StartTempoPoint = GetTempoPointForSeconds(InStartSeconds);
+	int32 ViewStartTick = StartTempoPoint->SecondsToTick(InStartSeconds);
+
+	const FFrameBasedTempoPoint* EndTempoPoint = GetTempoPointForSeconds(InEndSeconds);
+	int32 ViewEndTick = EndTempoPoint->SecondsToTick(InEndSeconds);
+
+	for (const FFrameBasedTempoPoint& TempoPoint : TempoMap)
+	{
+		if (TempoPoint.OnMs > ViewStartMs && TempoPoint.OnMs < ViewEndMs)
+		{
+			FTempoTimeSignatureMarker Marker;
+			Marker.Seconds = TempoPoint.OnMs / 1000.f;
+			Marker.TempoPoint = TempoPoint;
+
+			OutMarkers.Add(Marker);
+		}
+	}
+
+	for (const FFrameBasedTimeSignaturePoint& TimeSigPoint : BarMap)
+	{
+		if (TimeSigPoint.OnTick > ViewStartTick && TimeSigPoint.OnTick < ViewEndTick)
+		{
+			if (FTempoTimeSignatureMarker* FoundMarker = OutMarkers.FindByPredicate([TimeSigPoint](const FTempoTimeSignatureMarker& Marker) { return Marker.TempoPoint.IsSet() && Marker.TempoPoint->OnTick == TimeSigPoint.OnTick; }))
+			{
+				FoundMarker->TimeSigPoint = TimeSigPoint;
+			}
+			else
+			{
+				FTempoTimeSignatureMarker Marker;
+				Marker.TimeSigPoint = TimeSigPoint;
+
+				const FFrameBasedTempoPoint* TempoPoint = GetTempoPointForTick(TimeSigPoint.OnTick);
+				Marker.Seconds = TempoPoint->TickToSeconds(TimeSigPoint.OnTick);
+
+				OutMarkers.Add(Marker);
+			}
+		}
+	}
+}
+
+
 #undef LOCTEXT_NAMESPACE
diff --git a/Engine/Source/Runtime/MusicEnvironment/Public/FrameBasedMusicMap.h b/Engine/Source/Runtime/MusicEnvironment/Public/FrameBasedMusicMap.h
--- a/Engine/Source/Runtime/MusicEnvironment/Public/FrameBasedMusicMap.h
+++ b/Engine/Source/Runtime/MusicEnvironment/Public/FrameBasedMusicMap.h
@@ -56,11 +56,20 @@ namespace FrameBasedMusicMap
 	};
 }
 
-struct FSubdivisionPoint
+struct MUSICENVIRONMENT_API FGridInterval
 {
-	FFrameTime FrameTime;
+	double IntervalStartTime;
+
+	int32 IntervalTick;
+
 	FMusicalTime MusicalTime;
-	FString PrettyString;
+
+	FGridInterval(double InIntervalStartTime, int32 InIntervalTick, FMusicalTime InMusicalTime)
+	{
+		IntervalStartTime = InIntervalStartTime;
+		IntervalTick = InIntervalTick;
+		MusicalTime = InMusicalTime;
+	}
 };
 
 USTRUCT()
@@ -212,8 +221,13 @@ struct MUSICENVIRONMENT_API FFrameBasedTimeSignaturePoint
 		}
 		else
 		{
-			Result.Bar = OnBar + (BarsPassed - 1);
-			Result.TickInBar = TicksPast - ((-1 + BarsPassed) * TicksPerBar());
+			Result.Bar = OnBar + BarsPassed;
+			int32 TicksRemainder = FMath::Abs(TicksPast) % TicksPerBar();
+			if (TicksRemainder > 0)
+			{
+				Result.Bar--;
+				Result.TickInBar = TicksPerBar() - TicksRemainder;
+			}
 		}
 		return Result;
 	}
@@ -238,6 +252,21 @@ struct MUSICENVIRONMENT_API FFrameBasedTimeSignaturePoint
 	}
 };
 
+USTRUCT()
+struct MUSICENVIRONMENT_API FTempoTimeSignatureMarker
+{
+	GENERATED_BODY()
+
+	UPROPERTY()
+	TOptional<FFrameBasedTimeSignaturePoint> TimeSigPoint;
+
+	UPROPERTY()
+	TOptional<FFrameBasedTempoPoint> TempoPoint;
+
+	UPROPERTY()
+	float Seconds;
+};
+
 UCLASS()
 class MUSICENVIRONMENT_API UFrameBasedMusicMap : public UObject
 {
@@ -291,6 +320,7 @@ public:
 	FString MakeBarBeatStringForMusicalTime(const FMusicalTime& MusicalTime, FrameBasedMusicMap::EBarBeatFormat BarBeatFormat) const;
 	FString MakeBarBeatStringForSeconds(double Seconds, FrameBasedMusicMap::EBarBeatFormat BarBeatFormat) const;
 
+	FMusicalTime GetMusicalTime(int32 Tick) const;
 	FMusicalTime GetMusicalTime(const FFrameTime& FrameTime) const;
 
 	float GetTempoForFrame(double DecimalFrame) const;
@@ -343,7 +373,19 @@ public:
 	void AddTimeSignature(int32 AtTick, int32 AtBar, int16 Numerator, int16 Denominator);
 	void AddTempo(int32 AtTick, float AtMs, float Bpm);
 
-	TArray<FSubdivisionPoint> GatherSubdivisonsForTimespan(const FFrameTime& Start, const FFrameTime& End, const FFrameTime& MinSpacing, int32 QuantizationInterval) const;
+	void GatherSubdivisonsForTimespan(int32 StartTick, int32 EndTick, int32 QuantizationInterval, TArray<FGridInterval>& OutSubdivisions) const;
+
+	void GatherBarsForTimespan(int32 StartBar, int32 EndBar, int32 QuantizationInterval, TArray<FGridInterval>& OutBars) const;
+
+	void GetGridMetricsMusic(float InViewStartSeconds, float InViewEndSeconds, float PixelsPerSecond, float MinTickPx, float DesiredMajorTickPx, TArray<FGridInterval>& OutMajorIntervals, TArray<FGridInterval>& OutMinorIntervals) const;
+
+	void ComputeGridSpacingAtTick(int32 TickToCompute, float DesiredMajorTickSeconds, int32& OutMajorIntervalTicks) const;
+
+	void ComputeMinorGridSpacing(int32 StartTick, int32 EndTick, TArray<FGridInterval>& OutMinorIntervals) const;
+
+	TArray<FFrameBasedTempoPoint> GetTempoMetrics(float InViewStartSeconds, float InViewEndSeconds) const;
+
+	void GetTempoTimeSignatureMarkers(float InStartSeconds, float InEndSeconds, TArray<FTempoTimeSignatureMarker>& OutMarkers) const;
 
 private:
 	const FFrameBasedTimeSignaturePoint* GetTimeSignaturePointForTick(int32 Tick) const;
@@ -357,8 +399,6 @@ private:
 	const FFrameBasedTempoPoint* GetTempoPointForFrameTime(const FFrameTime& Frame) const;
 	const FFrameBasedTempoPoint* GetTempoPointForSeconds(double Seconds) const;
 
-	FSubdivisionPoint MakeSubdivisionPoint(const FMusicalTime& MusicalTime) const;
-
 	FString MakeBarBeatString(double DecimalFrame, const FFrameBasedTimeSignaturePoint* TimeSignature, const FFrameBasedTempoPoint* Tempo, FrameBasedMusicMap::EBarBeatFormat BarBeatFormat) const;
 };
 
