Fix CIS error caused by no default value for field

diff --git a/Engine/Source/Runtime/MusicEnvironment/Public/FrameBasedMusicMap.h b/Engine/Source/Runtime/MusicEnvironment/Public/FrameBasedMusicMap.h
--- a/Engine/Source/Runtime/MusicEnvironment/Public/FrameBasedMusicMap.h
+++ b/Engine/Source/Runtime/MusicEnvironment/Public/FrameBasedMusicMap.h
@@ -264,7 +264,7 @@ struct MUSICENVIRONMENT_API FTempoTimeSignatureMarker
 	TOptional<FFrameBasedTempoPoint> TempoPoint;
 
 	UPROPERTY()
-	float Seconds;
+	float Seconds = 0.f;
 };
 
 UCLASS()
