Fleshing out the new "music mode" for sequencer including:
- Audition metronome when playing in the sequence editor (via IMusicEnvironmentMetronome, FMetronomeTimeController, etc.)
- Real "external" musical clock source support (via FMusicClockSourceTimeController, etc.)
- Synchronization configuration options (via SMusicalClockConfigDialog)

diff --git a/Engine/Plugins/Runtime/Harmonix/Source/HarmonixMetasound/Private/HarmonixMetasound/Components/MusicClockComponent.cpp b/Engine/Plugins/Runtime/Harmonix/Source/HarmonixMetasound/Private/HarmonixMetasound/Components/MusicClockComponent.cpp
--- a/Engine/Plugins/Runtime/Harmonix/Source/HarmonixMetasound/Private/HarmonixMetasound/Components/MusicClockComponent.cpp
+++ b/Engine/Plugins/Runtime/Harmonix/Source/HarmonixMetasound/Private/HarmonixMetasound/Components/MusicClockComponent.cpp
@@ -790,32 +790,54 @@ void UMusicClockComponent::MakeDefaultSongMap()
 	DefaultMaps.GetBarMap().AddTimeSignatureAtBarIncludingCountIn(0, TimeSignatureNum, TimeSignatureDenom);
 }
 
-/*
-void UMusicClockComponent::OnTick(float DeltaSeconds, float InPlayRate)
-{
-
-}
-
-void UMusicClockComponent::OnStartPlaying(const FQualifiedFrameTime& InStartTime)
+FMusicalTime UMusicClockComponent::GetMusicalTime() const
 {
-
+	FMusicalTime Result;
+	if (CurrentVideoRenderSongPos.TimeSigDenominator < 1)
+	{
+		return Result;
+	}
+	Result.Bar = CurrentVideoRenderSongPos.Timestamp.Bar - 1;
+	Result.TicksPerBeat = (MusicalTime::TicksPerQuarterNote * 4 / CurrentVideoRenderSongPos.TimeSigDenominator);
+	Result.TicksPerBar = Result.TicksPerBeat * CurrentVideoRenderSongPos.TimeSigNumerator;
+	Result.TickInBar = static_cast<int32>((CurrentVideoRenderSongPos.Timestamp.Beat - 1.0f) * static_cast<float>(Result.TicksPerBeat));
+	return Result;
 }
 
-void UMusicClockComponent::OnStopPlaying(const FQualifiedFrameTime& InStopTime)
+FMusicalTime UMusicClockComponent::GetMusicalTimeWithSourceSpaceOffset(const FMusicalTime& Offset) const
 {
-
+	FMusicalTime Result;
+	const ISongMapEvaluator& MapEvaluator = GetSongMaps();
+	float CurrentTick = MapEvaluator.FractionalBarIncludingCountInToTick(CurrentVideoRenderSongPos.BarsIncludingCountIn);
+	
+	//float AdjustedTick = CurrentTick - (static_cast<float>(Offset.TickInBar * MapEvaluator.GetTicksPerQuarterNote()) / static_cast<float>(MusicalTime::TicksPerQuarterNote));
+	float OffsetTick = MapEvaluator.FractionalBarIncludingCountInToTick(static_cast<float>(Offset.FractionalBar()));
+	float AdjustedTick = CurrentTick - OffsetTick;
+
+	const FTimeSignature* TimeSignature = MapEvaluator.GetTimeSignatureAtTick(FMath::FloorToInt32(AdjustedTick));
+	FMusicTimestamp Timestamp = MapEvaluator.TickToMusicTimestamp(AdjustedTick);
+	Result.Bar = Timestamp.Bar - 1;
+	Result.TicksPerBeat = (MusicalTime::TicksPerQuarterNote * 4 / TimeSignature->Denominator);
+	Result.TicksPerBar = Result.TicksPerBeat * TimeSignature->Numerator;
+	Result.TickInBar = static_cast<int32>((Timestamp.Beat - 1.0f) * static_cast<float>(Result.TicksPerBeat));
+	return Result;
 }
 
-FMusicalTime UMusicClockComponent::GetMusicalTime() const
+FMusicalTime UMusicClockComponent::Quantize(const FMusicalTime& MusicalTime, int32 QuantizationInterval, UFrameBasedMusicMap::EQuantizeDirection Direction) const
 {
 	FMusicalTime Result;
-	Result.Bar = CurrentVideoRenderSongPos.Timestamp.Bar - 1;
-	Result.TicksPerBeat = (MusicalTime::TicksPerQuarterNote * 4 / CurrentVideoRenderSongPos.TimeSigDenominator);
-	Result.TicksPerBar = Result.TicksPerBeat * CurrentVideoRenderSongPos.TimeSigNumerator;
-	Result.TickInBar = static_cast<int32>((CurrentVideoRenderSongPos.Timestamp.Beat - 1.0f) * static_cast<float>(Result.TicksPerBeat));
+	const ISongMapEvaluator& MapEvaluator = GetSongMaps();
+	float Tick = MapEvaluator.FractionalBarIncludingCountInToTick(static_cast<float>(MusicalTime.FractionalBar()));
+	int32 QuantizedTick = MapEvaluator.QuantizeTickToNearestSubdivision(FMath::FloorToInt32(Tick), (EMidiFileQuantizeDirection)Direction, (EMidiClockSubdivisionQuantization) FrameBasedMusicMap::QuantiazationIntervalToQuartz(QuantizationInterval));
+	FMusicTimestamp Timestamp = MapEvaluator.TickToMusicTimestamp(QuantizedTick);
+	const FTimeSignature* TimeSignature = MapEvaluator.GetTimeSignatureAtTick(QuantizedTick);
+	check(TimeSignature);
+	Result.Bar = Timestamp.Bar - 1;
+	Result.TicksPerBeat = (MusicalTime::TicksPerQuarterNote * 4 / TimeSignature->Denominator);
+	Result.TicksPerBar = Result.TicksPerBeat * TimeSignature->Numerator;
+	Result.TickInBar = static_cast<int32>((Timestamp.Beat - 1.0f) * static_cast<float>(Result.TicksPerBeat));
 	return Result;
 }
-*/
 
 void FMusicClockDriverBase::EnsureClockIsValidForGameFrame()
 {
diff --git a/Engine/Plugins/Runtime/Harmonix/Source/HarmonixMetasound/Private/HarmonixMetasound/MusicEnvironmentSupport/HarmonixMusicEnvironmentMetronome.cpp b/Engine/Plugins/Runtime/Harmonix/Source/HarmonixMetasound/Private/HarmonixMetasound/MusicEnvironmentSupport/HarmonixMusicEnvironmentMetronome.cpp
--- a/Engine/Plugins/Runtime/Harmonix/Source/HarmonixMetasound/Private/HarmonixMetasound/MusicEnvironmentSupport/HarmonixMusicEnvironmentMetronome.cpp
+++ b/Engine/Plugins/Runtime/Harmonix/Source/HarmonixMetasound/Private/HarmonixMetasound/MusicEnvironmentSupport/HarmonixMusicEnvironmentMetronome.cpp
@@ -37,8 +37,8 @@ void UHarmonixMusicEnvironmentMetronome::Tick(float DeltaSecs)
 
 void UHarmonixMusicEnvironmentMetronome::Start(double FromSeconds)
 {
-	UE_LOG(LogMIDI, Log, TEXT("UHarmonixMusicEnvironmentMetronome: Starting Metronome at %f seconds"), FromSeconds);
-	
+	UE_LOG(LogMIDI, Verbose, TEXT("UHarmonixMusicEnvironmentMetronome: Starting Metronome at %f seconds"), FromSeconds);
+
 	if (!AudioComponent || !MusicClockComponent)
 	{
 		return;
@@ -57,8 +57,8 @@ void UHarmonixMusicEnvironmentMetronome::Seek(double ToSeconds)
 
 void UHarmonixMusicEnvironmentMetronome::Stop()
 {
-	UE_LOG(LogMIDI, Log, TEXT("UHarmonixMusicEnvironmentMetronome: Stopping Metronome"));
-	
+	UE_LOG(LogMIDI, Verbose, TEXT("UHarmonixMusicEnvironmentMetronome: Stopping Metronome"));
+
 	if (MusicClockComponent)
 	{
 		MusicClockComponent->Stop();
@@ -101,15 +101,33 @@ FMusicalTime UHarmonixMusicEnvironmentMetronome::GetMusicalTime() const
 	FMusicalTime MusicalTime;
 	if (MusicClockComponent)
 	{
-		const FMidiSongPos& SongPos = MusicClockComponent->GetSongPos();
-		MusicalTime.TicksPerBeat = MusicalTime::TicksPerQuarterNote * 4 / SongPos.TimeSigDenominator;
-		MusicalTime.TicksPerBar = MusicalTime.TicksPerBeat * SongPos.TimeSigNumerator;
-		MusicalTime.Bar = FMath::FloorToInt32(SongPos.BarsIncludingCountIn);
-		MusicalTime.TickInBar = FMath::FloorToInt32(static_cast<float>(MusicalTime.TicksPerBar) * FMath::Fractional(SongPos.BarsIncludingCountIn));
+		return MusicClockComponent->GetMusicalTime();
 	}
 	return MusicalTime;
 }
 
+FMusicalTime UHarmonixMusicEnvironmentMetronome::GetMusicalTimeWithSourceSpaceOffset(const FMusicalTime& Offset) const
+{
+	FMusicalTime MusicalTime;
+	if (MusicClockComponent)
+	{
+		return MusicClockComponent->GetMusicalTimeWithSourceSpaceOffset(Offset);
+	}
+	return MusicalTime;
+}
+
+FMusicalTime UHarmonixMusicEnvironmentMetronome::Quantize(const FMusicalTime& MusicalTime, int32 QuantizationInterval, UFrameBasedMusicMap::EQuantizeDirection Direction) const
+{
+	if (!MusicClockComponent)
+	{
+		return MusicalTime;
+	}
+	else
+	{
+		return MusicClockComponent->Quantize(MusicalTime, QuantizationInterval, Direction);
+	}
+}
+
 void UHarmonixMusicEnvironmentMetronome::BeginDestroy()
 {
 	if (MusicClockComponent)
@@ -141,11 +159,11 @@ bool UHarmonixMusicEnvironmentMetronome::BuildAndStartMetasound(UWorld* InWorld)
 	{
 		return false;
 	}
-	AudioComponent = NewObject<UAudioComponent>(this);
+	AudioComponent = NewObject<UAudioComponent>(this, "MetronomeAudio", RF_Transient);
 	AudioComponent->SetUISound(true);
 	SourceBuilder->Audition(this, AudioComponent, {});
 
-	MusicClockComponent = NewObject<UMusicClockComponent>(this); 
+	MusicClockComponent = NewObject<UMusicClockComponent>(this, "MetronomeClock", RF_Transient);
 	MusicClockComponent->MetasoundOutputName = TEXT("MIDI Clock");
 	MusicClockComponent->ConnectToMetasoundOnAudioComponent(AudioComponent);
 	MusicClockComponent->Start();
diff --git a/Engine/Plugins/Runtime/Harmonix/Source/HarmonixMetasound/Private/HarmonixMetasound/MusicEnvironmentSupport/HarmonixMusicEnvironmentMetronome.h b/Engine/Plugins/Runtime/Harmonix/Source/HarmonixMetasound/Private/HarmonixMetasound/MusicEnvironmentSupport/HarmonixMusicEnvironmentMetronome.h
--- a/Engine/Plugins/Runtime/Harmonix/Source/HarmonixMetasound/Private/HarmonixMetasound/MusicEnvironmentSupport/HarmonixMusicEnvironmentMetronome.h
+++ b/Engine/Plugins/Runtime/Harmonix/Source/HarmonixMetasound/Private/HarmonixMetasound/MusicEnvironmentSupport/HarmonixMusicEnvironmentMetronome.h
@@ -35,7 +35,11 @@ public:
 
 	virtual float GetCurrentTempo() override;
 	virtual double GetCurrentPositionSeconds() override;
+
 	virtual FMusicalTime GetMusicalTime() const override;
+	virtual FMusicalTime GetMusicalTimeWithSourceSpaceOffset(const FMusicalTime& Offset) const override;
+	virtual FMusicalTime Quantize(const FMusicalTime& MusicalTime, int32 QuantizationInterval, UFrameBasedMusicMap::EQuantizeDirection Direction = UFrameBasedMusicMap::EQuantizeDirection::Nearest) const override;
+	virtual bool CanAuditionInEditor() const override { return true; }
 
 	virtual void BeginDestroy() override;
 	
@@ -44,9 +48,6 @@ protected:
 	virtual bool OnSetTempo(float Bpm) override;
 	
 private:
-	UPROPERTY(Transient)
-	TObjectPtr<AActor> Actor;
-	
 	UPROPERTY(Transient)
 	TObjectPtr<UAudioComponent> AudioComponent;
 	UPROPERTY(Transient)
diff --git a/Engine/Plugins/Runtime/Harmonix/Source/HarmonixMetasound/Public/HarmonixMetasound/Components/MusicClockComponent.h b/Engine/Plugins/Runtime/Harmonix/Source/HarmonixMetasound/Public/HarmonixMetasound/Components/MusicClockComponent.h
--- a/Engine/Plugins/Runtime/Harmonix/Source/HarmonixMetasound/Public/HarmonixMetasound/Components/MusicClockComponent.h
+++ b/Engine/Plugins/Runtime/Harmonix/Source/HarmonixMetasound/Public/HarmonixMetasound/Components/MusicClockComponent.h
@@ -2,12 +2,13 @@
 #pragma once
 
 #include "Components/ActorComponent.h"
+#include "Delegates/DelegateCombinations.h"
+#include "Evaluation/IMovieSceneCustomClockSource.h"
 #include "Harmonix/MusicalTimebase.h"
 #include "HarmonixMidi/MidiSongPos.h"
 #include "HarmonixMidi/SongMaps.h"
-#include "Delegates/DelegateCombinations.h"
+#include "MusicEnvironmentClockSource.h"
 #include "Templates/UniquePtr.h"
-#include "Evaluation/IMovieSceneCustomClockSource.h"
 
 #include "MusicClockComponent.generated.h"
 
@@ -50,7 +51,7 @@ DECLARE_DYNAMIC_MULTICAST_DELEGATE(FMusicClockConnected);
 DECLARE_DYNAMIC_MULTICAST_DELEGATE(FMusicClockDisconnected);
 
 UCLASS(ClassGroup = (MetaSoundMusic), PrioritizeCategories = "MusicClock", meta = (BlueprintSpawnableComponent, DisplayName = "Music Clock", ScriptName = MusicClockComponent))
-class HARMONIXMETASOUND_API UMusicClockComponent : public UActorComponent/*, public IMovieSceneMusicClockSource*/
+class HARMONIXMETASOUND_API UMusicClockComponent : public UActorComponent, public IMusicEnvironmentClockSource
 {
 	GENERATED_BODY()
 
@@ -374,15 +375,12 @@ public:
 
 	const FMidiSongPos& GetRawUnsmoothedAudioRenderPos() const { return CurrentRawAudioRenderSongPos; }
 
-	// vvv BEGIN IMovieSceneMusicClockSource Implementation vvv
-	/*
-	virtual void OnTick(float DeltaSeconds, float InPlayRate) override;
-	virtual void OnStartPlaying(const FQualifiedFrameTime& InStartTime) override;
-	virtual void OnStopPlaying(const FQualifiedFrameTime& InStopTime) override;
+	// vvv BEGIN  IMusicEnvironmentClockSource Implementation vvv
 	virtual FMusicalTime GetMusicalTime() const override;
-	*/
-	// ^^^ END IMovieSceneMusicClockSource Implementation ^^^
-
+	virtual FMusicalTime GetMusicalTimeWithSourceSpaceOffset(const FMusicalTime& Offset) const;
+	virtual FMusicalTime Quantize(const FMusicalTime& MusicalTime, int32 QuantizationInterval, UFrameBasedMusicMap::EQuantizeDirection Direction = UFrameBasedMusicMap::EQuantizeDirection::Nearest) const;
+	virtual bool CanAuditionInEditor() const override { return false; }
+	// ^^^ END IMusicEnvironmentClockSource Implementation ^^^
 
 public:
 	virtual void TickComponent(float DeltaTime, enum ELevelTick TickType, FActorComponentTickFunction* ThisTickFunction);
diff --git a/Engine/Plugins/Runtime/Harmonix/Source/HarmonixMidi/Public/HarmonixMidi/SongMaps.h b/Engine/Plugins/Runtime/Harmonix/Source/HarmonixMidi/Public/HarmonixMidi/SongMaps.h
--- a/Engine/Plugins/Runtime/Harmonix/Source/HarmonixMidi/Public/HarmonixMidi/SongMaps.h
+++ b/Engine/Plugins/Runtime/Harmonix/Source/HarmonixMidi/Public/HarmonixMidi/SongMaps.h
@@ -25,9 +25,9 @@ namespace HarmonixMetasound
 UENUM()
 enum class EMidiFileQuantizeDirection : uint8
 {
-	Nearest,
-	Up,
-	Down
+	Nearest = 0,
+	Up = 1,
+	Down = 2
 };
 
 UENUM(BlueprintType)
diff --git a/Engine/Source/Runtime/Core/Public/Misc/MusicalTime.h b/Engine/Source/Runtime/Core/Public/Misc/MusicalTime.h
--- a/Engine/Source/Runtime/Core/Public/Misc/MusicalTime.h
+++ b/Engine/Source/Runtime/Core/Public/Misc/MusicalTime.h
@@ -41,13 +41,21 @@ struct FMusicalTime
 		return TicksPerBar > 0 && TicksPerBeat > 0;
 	}
 
-	float BeatInBar() const
+	double FractionalBeatInBar() const
 	{
-		return static_cast<float>(TickInBar) / static_cast<float>(TicksPerBeat);
+		if (!IsValid())
+		{
+			return 0.0;
+		}
+		return static_cast<double>(TickInBar) / static_cast<double>(TicksPerBeat);
 	}
 
 	double FractionalBar() const
 	{
+		if (!IsValid())
+		{
+			return 0.0;
+		}
 		return static_cast<double>(Bar) + (static_cast<double>(TickInBar) / static_cast<double>(TicksPerBar));
 	}
 
diff --git a/Engine/Source/Runtime/MusicEnvironment/Private/FrameBasedMusicMap.cpp b/Engine/Source/Runtime/MusicEnvironment/Private/FrameBasedMusicMap.cpp
--- a/Engine/Source/Runtime/MusicEnvironment/Private/FrameBasedMusicMap.cpp
+++ b/Engine/Source/Runtime/MusicEnvironment/Private/FrameBasedMusicMap.cpp
@@ -36,7 +36,38 @@ namespace FrameBasedMusicMap
 		case EDottedQuantizationTicks::DottedHalfNote:			return LOCTEXT("DottedHalfNote", "Dotted Half Note");
 		case EDottedQuantizationTicks::DottedWholeNote:			return LOCTEXT("DottedWholeNote", "Dotted Whole Note");
 		}
-		return LOCTEXT("Unknown", "Unrecognized Qunatization Factor");
+		return LOCTEXT("Unknown", "Unrecognized Quantization Factor");
+	}
+
+	EQuartzCommandQuantization QuantiazationIntervalToQuartz(int32 Quantization)
+	{
+		switch (Quantization)
+		{
+		case ETripletQuantizationTicks::ThirtysecondNoteTriplet:	return EQuartzCommandQuantization::None; // quartz doesn't go this small.
+		case ETripletQuantizationTicks::SixteenthNoteTriplet:		return EQuartzCommandQuantization::SixteenthNoteTriplet;
+		case ETripletQuantizationTicks::EighthNoteTriplet:			return EQuartzCommandQuantization::EighthNoteTriplet;
+		case ETripletQuantizationTicks::QuarterNoteTriplet:			return EQuartzCommandQuantization::QuarterNoteTriplet;
+		case ETripletQuantizationTicks::HalfNoteTriplet:			return EQuartzCommandQuantization::HalfNoteTriplet;
+		case ETripletQuantizationTicks::WholeNoteTriplet:			return EQuartzCommandQuantization::None; // quartz doesn't support this division.
+
+		case EWholeQuantizationTicks::ThirtysecondNote:	return EQuartzCommandQuantization::ThirtySecondNote;
+		case EWholeQuantizationTicks::SixteenthNote:	return EQuartzCommandQuantization::SixteenthNote;
+		case EWholeQuantizationTicks::EighthNote:		return EQuartzCommandQuantization::EighthNote;
+		case EWholeQuantizationTicks::QuarterNote:		return EQuartzCommandQuantization::QuarterNote;
+		case EWholeQuantizationTicks::HalfNote:			return EQuartzCommandQuantization::HalfNote;
+		case EWholeQuantizationTicks::WholeNote:		return EQuartzCommandQuantization::WholeNote;
+		case EWholeQuantizationTicks::Bar:				return EQuartzCommandQuantization::Bar;
+		case EWholeQuantizationTicks::Beat:				return EQuartzCommandQuantization::Beat;
+
+
+		case EDottedQuantizationTicks::DottedThirtysecondNote:	return EQuartzCommandQuantization::None; // quartz doesn't support this division.
+		case EDottedQuantizationTicks::DottedSixteenthNote:		return EQuartzCommandQuantization::DottedSixteenthNote;
+		case EDottedQuantizationTicks::DottedEighthNote:		return EQuartzCommandQuantization::DottedEighthNote;
+		case EDottedQuantizationTicks::DottedQuarterNote:		return EQuartzCommandQuantization::DottedQuarterNote;
+		case EDottedQuantizationTicks::DottedHalfNote:			return EQuartzCommandQuantization::DottedHalfNote;
+		case EDottedQuantizationTicks::DottedWholeNote:			return EQuartzCommandQuantization::DottedWholeNote;
+		}
+		return EQuartzCommandQuantization::None;
 	}
 
 	TArray<TSharedPtr<int32>> QuantizationOptionsSource = {
@@ -134,6 +165,18 @@ FFrameTime UFrameBasedMusicMap::GetFrameTime(int32 AtBar, float AtBeat) const
 	return FrameResolution.AsFrameTime(GetSeconds(AtBar, AtBeat));
 }
 
+FFrameTime UFrameBasedMusicMap::GetFrameTime(double FractionalBar) const
+{
+	if (IsEmpty())
+	{
+		return FFrameTime();
+	}
+
+	int32 Bar = FMath::FloorToInt32(FractionalBar);
+	float Beat =static_cast<float>(static_cast<double>(GetBeatsInBar(Bar)) * FMath::Fractional(FractionalBar));
+	return GetFrameTime(Bar, Beat);	
+}
+
 double UFrameBasedMusicMap::GetSeconds(const FMusicalTime& MusicalTime) const
 {
 	if (IsEmpty())
@@ -242,6 +285,30 @@ FFrameBasedTimeSignature UFrameBasedMusicMap::GetTimeSignatureForSeconds(double
 	return Point ? Point->TimeSignature : FFrameBasedTimeSignature();
 }
 
+int32 UFrameBasedMusicMap::GetBeatsInBar(int32 Bar) const
+{
+	if (BarMap.IsEmpty())
+	{
+		return 4;
+	}
+	
+	const FFrameBasedTimeSignaturePoint* TimeSignaturePoint = GetTimeSignaturePointForBar(Bar);
+	check(TimeSignaturePoint);
+	return TimeSignaturePoint->TimeSignature.Numerator;
+}
+
+int32 UFrameBasedMusicMap::GetTicksInBar(int32 Bar) const
+{
+	if (BarMap.IsEmpty())
+	{
+		return 4 * MusicalTime::TicksPerQuarterNote;
+	}
+	
+	const FFrameBasedTimeSignaturePoint* TimeSignaturePoint = GetTimeSignaturePointForBar(Bar);
+	check(TimeSignaturePoint);
+	return TimeSignaturePoint->TicksPerBar();	
+}
+
 const FFrameBasedTimeSignaturePoint* UFrameBasedMusicMap::GetTimeSignaturePointForTick(int32 Tick) const
 {
 	if (BarMap.IsEmpty())
@@ -392,6 +459,15 @@ FMusicalTime UFrameBasedMusicMap::BarsToMusicalTime(double BarsFromStart) const
 }
 
 FMusicalTime UFrameBasedMusicMap::Quantize(const FMusicalTime& MusicalTime, int32 QuantizationInterval, EQuantizeDirection Direction) const
+{
+	int32 QuantizedTick = QuantizeToTick(MusicalTime, QuantizationInterval, Direction);
+
+	const FFrameBasedTimeSignaturePoint* TimeSignaturePoint = GetTimeSignaturePointForTick(QuantizedTick);
+
+	return TimeSignaturePoint->TickToMusicalTime(QuantizedTick);
+}
+
+int32 UFrameBasedMusicMap::QuantizeToTick(const FMusicalTime& MusicalTime, int32 QuantizationInterval, EQuantizeDirection Direction /*= EQuantizeDirection::Nearest*/) const
 {
 	int32 TicksInQuantizationInterval = QuantizationInterval;
 	switch (QuantizationInterval)
@@ -429,9 +505,7 @@ FMusicalTime UFrameBasedMusicMap::Quantize(const FMusicalTime& MusicalTime, int3
 
 	int32 FinalTick = TimeSignaturePoint->OnTick + QuantizedTicksInTimeSignature;
 
-	TimeSignaturePoint = GetTimeSignaturePointForTick(FinalTick);
-
-	return TimeSignaturePoint->TickToMusicalTime(FinalTick);
+	return FinalTick;
 }
 
 FMusicalTime UFrameBasedMusicMap::Add(const FMusicalTime& MusicalTime, int32 QuantizationInterval) const
@@ -555,7 +629,7 @@ FString UFrameBasedMusicMap::MakeBarBeatString(double DecimalFrame, const FFrame
 	FMusicalTime MusicalTime = TimeSignature->TickToMusicalTime(Tick);
 	// for display to users... bars and beats are 1 based!
 	MusicalTime.Bar += (BarBeatFormat == FrameBasedMusicMap::EBarBeatFormat::Position ? 1 : 0);
-	double BeatAsDouble = FMath::RoundToDouble((MusicalTime.BeatInBar() + (BarBeatFormat == FrameBasedMusicMap::EBarBeatFormat::Position ? 1.0f : 0.0f)) * 100.0f);
+	double BeatAsDouble = FMath::RoundToDouble((static_cast<float>(MusicalTime.FractionalBeatInBar()) + (BarBeatFormat == FrameBasedMusicMap::EBarBeatFormat::Position ? 1.0f : 0.0f)) * 100.0f);
 	if (BeatAsDouble >= ((TimeSignature->TimeSignature.Numerator + (BarBeatFormat == FrameBasedMusicMap::EBarBeatFormat::Position ? 1 : 0)) * 100))
 	{
 		MusicalTime.Bar++;
diff --git a/Engine/Source/Runtime/MusicEnvironment/Private/MusicEvironmentSubsystem.cpp b/Engine/Source/Runtime/MusicEnvironment/Private/MusicEvironmentSubsystem.cpp
--- a/Engine/Source/Runtime/MusicEnvironment/Private/MusicEvironmentSubsystem.cpp
+++ b/Engine/Source/Runtime/MusicEnvironment/Private/MusicEvironmentSubsystem.cpp
@@ -38,11 +38,11 @@ bool UMusicEnvironmentSubsystem::CanSpawnMetronome() const
 	return MetronomeType.IsValid();
 }
 
-TScriptInterface<IMusicEnvironmentMetronome> UMusicEnvironmentSubsystem::SpawnMetronome(UObject* Outer)
+TScriptInterface<IMusicEnvironmentMetronome> UMusicEnvironmentSubsystem::SpawnMetronome(UObject* Outer, FName Name)
 {
 	if (!MetronomeType.IsValid())
 	{
 		return nullptr;
 	}
-	return TScriptInterface<IMusicEnvironmentMetronome>(NewObject<UObject>(Outer, MetronomeType.Get()));
+	return TScriptInterface<IMusicEnvironmentMetronome>(NewObject<UObject>(Outer, MetronomeType.Get(), Name, RF_Transient));
 }
diff --git a/Engine/Source/Runtime/MusicEnvironment/Public/FrameBasedMusicMap.h b/Engine/Source/Runtime/MusicEnvironment/Public/FrameBasedMusicMap.h
--- a/Engine/Source/Runtime/MusicEnvironment/Public/FrameBasedMusicMap.h
+++ b/Engine/Source/Runtime/MusicEnvironment/Public/FrameBasedMusicMap.h
@@ -6,6 +6,7 @@
 #include "Misc/FrameRate.h"
 #include "Misc/MusicalTime.h"
 #include "UObject/Interface.h"
+#include "Sound/QuartzQuantizationUtilities.h"
 
 #include "FrameBasedMusicMap.generated.h"
 
@@ -46,6 +47,8 @@ namespace FrameBasedMusicMap
 	};
 
 	MUSICENVIRONMENT_API FText GetQuantizationName(int32 Quantization);
+	MUSICENVIRONMENT_API EQuartzCommandQuantization QuantiazationIntervalToQuartz(int32 Quantization);
+
 	extern MUSICENVIRONMENT_API TArray<TSharedPtr<int32>> QuantizationOptionsSource;
 
 	enum class EBarBeatFormat
@@ -312,6 +315,7 @@ public:
 	double GetFrame(int32 AtBar) const;
 	FFrameTime GetFrameTime(const FMusicalTime& MusicalTime) const;
 	FFrameTime GetFrameTime(int32 AtBar, float AtBeat) const;
+	FFrameTime GetFrameTime(double FractionalBar) const;
 	double GetSeconds(const FMusicalTime& MusicalTime) const;
 	double GetSeconds(int32 AtBar, float AtBeat = 0.0f) const;
 
@@ -330,18 +334,21 @@ public:
 	FFrameBasedTimeSignature GetTimeSignatureForFrame(double DecimalFrame) const;
 	FFrameBasedTimeSignature GetTimeSignatureForFrameTime(const FFrameTime& FrameTime) const;
 	FFrameBasedTimeSignature GetTimeSignatureForSeconds(double Seconds) const;
-
+	int32 GetBeatsInBar(int32 Bar) const;
+	int32 GetTicksInBar(int32 Bar) const;
+	
 	FFrameTime RemapFrameTime(FFrameTime SourceTime, const UFrameBasedMusicMap* OldMap) const;
 
 	FMusicalTime BarsToMusicalTime(double BarsFromStart) const;
 
 	enum class EQuantizeDirection
 	{
-		Nearest,
-		Ceiling,
-		Floor
+		Nearest = 0,
+		Ceiling = 1,
+		Floor = 2
 	};
 	FMusicalTime Quantize(const FMusicalTime& MusicalTime, int32 QuantizationInterval, EQuantizeDirection Direction = EQuantizeDirection::Nearest) const;
+	int32 QuantizeToTick(const FMusicalTime& MusicalTime, int32 QuantizationInterval, EQuantizeDirection Direction = EQuantizeDirection::Nearest) const;
 
 	/**
 	* FMusicalTime Add(const FMusicalTime& MusicalTime, int32 QuantizationInterval) const;
diff --git a/Engine/Source/Runtime/MusicEnvironment/Public/MusicEnvironmentClockSource.h b/Engine/Source/Runtime/MusicEnvironment/Public/MusicEnvironmentClockSource.h
--- a/Engine/Source/Runtime/MusicEnvironment/Public/MusicEnvironmentClockSource.h
+++ b/Engine/Source/Runtime/MusicEnvironment/Public/MusicEnvironmentClockSource.h
@@ -3,6 +3,7 @@
 #pragma once
 
 #include "CoreMinimal.h"
+#include "FrameBasedMusicMap.h"
 #include "Misc/MusicalTime.h"
 #include "UObject/Interface.h"
 
@@ -10,7 +11,7 @@
 
 // This class does not need to be modified.
 UINTERFACE(BlueprintType, meta = (CannotImplementInterfaceInBlueprint))
-class UMusicEnvironmentClockSource : public UInterface
+class MUSICENVIRONMENT_API UMusicEnvironmentClockSource : public UInterface
 {
 	GENERATED_BODY()
 };
@@ -24,4 +25,8 @@ class MUSICENVIRONMENT_API IMusicEnvironmentClockSource
 
 public:
 	virtual FMusicalTime GetMusicalTime() const = 0;
+	//virtual FMusicalTime GetBarOneBeatOneMusicalTime() const = 0; // returns bar 0 beat 0 But ticks per beat and ticks per bar are the initial tempo!
+	virtual FMusicalTime GetMusicalTimeWithSourceSpaceOffset(const FMusicalTime& Offset) const = 0;
+	virtual FMusicalTime Quantize(const FMusicalTime& MusicalTime, int32 QuantizationInterval, UFrameBasedMusicMap::EQuantizeDirection Direction = UFrameBasedMusicMap::EQuantizeDirection::Nearest) const = 0;
+	virtual bool CanAuditionInEditor() const = 0;
 };
diff --git a/Engine/Source/Runtime/MusicEnvironment/Public/MusicEnvironmentSubsystem.h b/Engine/Source/Runtime/MusicEnvironment/Public/MusicEnvironmentSubsystem.h
--- a/Engine/Source/Runtime/MusicEnvironment/Public/MusicEnvironmentSubsystem.h
+++ b/Engine/Source/Runtime/MusicEnvironment/Public/MusicEnvironmentSubsystem.h
@@ -31,7 +31,7 @@ public:
 	bool CanSpawnMetronome() const;
 
 	UFUNCTION(BlueprintCallable, Category = "Music|Clock")
-	virtual TScriptInterface<IMusicEnvironmentMetronome> SpawnMetronome(UObject* Outer);
+	virtual TScriptInterface<IMusicEnvironmentMetronome> SpawnMetronome(UObject* Outer, FName Name = NAME_None);
 
 private:
 	UPROPERTY(Transient)
diff --git a/Engine/Source/Runtime/MusicEnvironment/Public/MusicTypes/MusicalEnums.h b/Engine/Source/Runtime/MusicEnvironment/Public/MusicTypes/MusicalEnums.h
new file mode 100644
--- /dev/null
+++ b/Engine/Source/Runtime/MusicEnvironment/Public/MusicTypes/MusicalEnums.h
@@ -0,0 +1,24 @@
+// Copyright Epic Games, Inc. All Rights Reserved.
+
+#pragma once
+
+#include "MusicalEnums.generated.h"
+
+UENUM(BlueprintType)
+enum class EMusicalSyncPointCaptureMode : uint8
+{
+	DirectMapped              UMETA(ToolTip="Bar 1 | Beat 1 of the level sequence is mapped directly to Bar 1 | Beat 1 of the driving music clock. If needed the sequence is 'seeked' to match the current time of the music clock."),
+	Immediate                 UMETA(ToolTip="Bar 1 | Beat 1 of the level sequence will be mapped to the current position of the music clock. Now quantization is done. While the sequence will advance in a locked-step, tempo matched way, its offset will be maintained. CAREFUL! You probably want to use some sort of qunatized capture."),
+	NearestQuantizedPosition  UMETA(ToolTip="Bar 1 | Beat 1 of the level sequence is mapped to the NEAREST specified quantized musical subdivision. May be before or after the current musical time. Playback will start right away if the nearest subdivision is in the past, otherwise playback will start when the next subdivision arrives."),
+	NextQuantizedPosition     UMETA(ToolTip="Bar 1 | Beat 1 of the level sequence is mapped to the NEXT specified quantized musical subdivision. Will be on or after the current musical time. Playback will start when that subdivision arrives."),
+	PreviousQuantizedPosition UMETA(ToolTip="Bar 1 | Beat 1 of the level sequence is mapped to the Previous specified quantized musical subdivision. Will be on or before the current musical time. Playback will start right away."),
+};
+
+UENUM(BlueprintType)
+enum class EMusicalTimeSignatureRemapMode : uint8
+{
+	None                    UMETA(ToolTip="No attempt is made to maintain 'bar synchronizztion' if there are mismatched time signatured."),
+	ScaleBars               UMETA(ToolTip="Scale each bar as neccessary to match bar lengths in music clock."),
+	CropOrLoopBars          UMETA(ToolTip="Crop or loop bars as nessary to align downbeats between the level sequence and the music clock."),
+	ScaleBeatsAndCropOrLoop UMETA(ToolTip="Scale beats to match music clock beat lengths, and THEN crop or loop bars as nessary to align downbeats between the level sequence and the music clock.")
+};
