Fix file name typo.

diff --git a/Engine/Source/Runtime/MusicEnvironment/Private/MusicEnvironmentSubsystem.cpp b/Engine/Source/Runtime/MusicEnvironment/Private/MusicEnvironmentSubsystem.cpp
new file mode 100644
--- /dev/null
+++ b/Engine/Source/Runtime/MusicEnvironment/Private/MusicEnvironmentSubsystem.cpp
@@ -0,0 +1,48 @@
+// Copyright Epic Games, Inc. All Rights Reserved.
+
+#include "MusicEnvironmentSubsystem.h"
+
+#include "MusicClockSourceManager.h"
+#include "MusicEnvironmentMetronome.h"
+
+#include "Engine/Engine.h"
+
+DEFINE_LOG_CATEGORY(LogMusicEnvironment);
+
+UMusicEnvironmentSubsystem& UMusicEnvironmentSubsystem::Get()
+{
+	return *GEngine->GetEngineSubsystem< UMusicEnvironmentSubsystem>();
+}
+
+UMusicClockSourceManager* UMusicEnvironmentSubsystem::GetClockSourceManager()
+{
+	if (!ClockSourceManager)
+	{
+		ClockSourceManager = NewObject<UMusicClockSourceManager>(this);
+	}
+	return ClockSourceManager;
+}
+
+bool UMusicEnvironmentSubsystem::SetMetronomeClass(const UClass* InMetronomeType)
+{
+	if (!InMetronomeType->ImplementsInterface(UMusicEnvironmentMetronome::StaticClass()))
+	{
+		return false;
+	}
+	MetronomeType = InMetronomeType;
+	return true;
+}
+
+bool UMusicEnvironmentSubsystem::CanSpawnMetronome() const
+{
+	return MetronomeType.IsValid();
+}
+
+TScriptInterface<IMusicEnvironmentMetronome> UMusicEnvironmentSubsystem::SpawnMetronome(UObject* Outer, FName Name)
+{
+	if (!MetronomeType.IsValid())
+	{
+		return nullptr;
+	}
+	return TScriptInterface<IMusicEnvironmentMetronome>(NewObject<UObject>(Outer, MetronomeType.Get(), Name, RF_Transient));
+}
diff --git a/Engine/Source/Runtime/MusicEnvironment/Private/MusicEvironmentSubsystem.cpp b/Engine/Source/Runtime/MusicEnvironment/Private/MusicEvironmentSubsystem.cpp
deleted file mode 100644
