Fix issue with timeline tick marks not displaying correctly for odd time signatures.

diff --git a/Engine/Source/Runtime/MusicEnvironment/Private/FrameBasedMusicMap.cpp b/Engine/Source/Runtime/MusicEnvironment/Private/FrameBasedMusicMap.cpp
--- a/Engine/Source/Runtime/MusicEnvironment/Private/FrameBasedMusicMap.cpp
+++ b/Engine/Source/Runtime/MusicEnvironment/Private/FrameBasedMusicMap.cpp
@@ -846,7 +846,10 @@ void UFrameBasedMusicMap::ComputeGridSpacingAtTick(int32 TickToCompute, float De
 void UFrameBasedMusicMap::ComputeMinorGridSpacing(int32 StartTick, int32 EndTick, TArray<FGridInterval>& OutMinorIntervals) const
 {
 	int32 StepTickSize = EndTick - StartTick;
-
+    if (StepTickSize <= 0)
+    {
+    	return;
+    }
 	FMusicalTime StartMusicalTime = GetMusicalTime(StartTick);
 	FMusicalTime EndMusicalTime = GetMusicalTime(EndTick);
 
@@ -859,7 +862,11 @@ void UFrameBasedMusicMap::ComputeMinorGridSpacing(int32 StartTick, int32 EndTick
 	}
 	else
 	{
-		int32 MinorTickSize = StepTickSize / 4;
+		int32 MinorTickSize = StartMusicalTime.TicksPerBeat;
+		while (MinorTickSize >= StepTickSize)
+		{
+			MinorTickSize /= 2;
+		}
 		MinorTickSize = FMath::Clamp(MinorTickSize, FrameBasedMusicMap::EWholeQuantizationTicks::ThirtysecondNote, FrameBasedMusicMap::EWholeQuantizationTicks::HalfNote);
 		GatherSubdivisonsForTimespan(StartTick, EndTick, MinorTickSize, OutMinorIntervals);
 	}
