Make new metronome objects with a unique name to avoid issues with refcounted objects trying to get recycled by the UObject system

#rnx

diff --git a/Engine/Source/Runtime/MusicEnvironment/Private/MusicEnvironmentSubsystem.cpp b/Engine/Source/Runtime/MusicEnvironment/Private/MusicEnvironmentSubsystem.cpp
--- a/Engine/Source/Runtime/MusicEnvironment/Private/MusicEnvironmentSubsystem.cpp
+++ b/Engine/Source/Runtime/MusicEnvironment/Private/MusicEnvironmentSubsystem.cpp
@@ -44,5 +44,5 @@ TScriptInterface<IMusicEnvironmentMetronome> UMusicEnvironmentSubsystem::SpawnMe
 	{
 		return nullptr;
 	}
-	return TScriptInterface<IMusicEnvironmentMetronome>(NewObject<UObject>(Outer, MetronomeType.Get(), Name, RF_Transient));
+	return TScriptInterface<IMusicEnvironmentMetronome>(NewObject<UObject>(Outer, MetronomeType.Get(), MakeUniqueObjectName(Outer, MetronomeType.Get(), Name), RF_Transient));
 }
